Fix using case-insensitive include which causes linux/mac builds to fail when trying to compile MirrorDataTableEditorModule.

#rnx
#Synced-CL 52433329

diff --git a/Engine/Source/Editor/MirrorDataTableEditor/Private/MirrorDataTableEditorModule.cpp b/Engine/Source/Editor/MirrorDataTableEditor/Private/MirrorDataTableEditorModule.cpp
--- a/Engine/Source/Editor/MirrorDataTableEditor/Private/MirrorDataTableEditorModule.cpp
+++ b/Engine/Source/Editor/MirrorDataTableEditor/Private/MirrorDataTableEditorModule.cpp
@@ -13,7 +13,7 @@
 #include "ToolMenu.h"
 #include "ToolMenus.h"
 #include "Animation/AnimationSettings.h"
-#include "Animation/AnimCurveMetaData.h"
+#include "Animation/AnimCurveMetadata.h"
 #include "Animation/MirrorDataTable.h"
 #include "Framework/Commands/Commands.h"
 #include "Framework/Notifications/NotificationManager.h"
