Added mirror actions to the context menu when right clicking a mirror datatable row

#rnx
#Synced-CL 52485170

diff --git a/Engine/Source/Editor/MirrorDataTableEditor/Private/MirrorDataTableEditorModule.cpp b/Engine/Source/Editor/MirrorDataTableEditor/Private/MirrorDataTableEditorModule.cpp
--- a/Engine/Source/Editor/MirrorDataTableEditor/Private/MirrorDataTableEditorModule.cpp
+++ b/Engine/Source/Editor/MirrorDataTableEditor/Private/MirrorDataTableEditorModule.cpp
@@ -4,6 +4,7 @@
 
 #include "DataTableEditorModule.h"
 #include "DataTableEditorUtils.h"
+#include "DataTableRowMenuContext.h"
 #include "IEditableSkeleton.h"
 #include "ISkeletonEditorModule.h"
 #include "MirrorDataTableCustomization.h"
@@ -44,6 +45,24 @@ static UMirrorDataTable* GetMirrorDataTableFromContext(const FToolMenuContext& C
 	return nullptr;
 }
 
+/** Returns the selected row in the mirror data table currently being edited, if any. */
+static FMirrorTableRow* GetMirrorTableRowFromContext(const FToolMenuContext& Context)
+{
+	UMirrorDataTable* MirrorTable = GetMirrorDataTableFromContext(Context);
+	if (!MirrorTable)
+	{
+		return nullptr;
+	}
+
+	const UDataTableRowMenuContext* RowCtx = Context.FindContext<UDataTableRowMenuContext>();
+	if (!RowCtx || !RowCtx->RowDataPtr)
+	{
+		return nullptr;
+	}
+
+	return MirrorTable->FindRow<FMirrorTableRow>(RowCtx->RowDataPtr->RowId, TEXT(""));
+}
+
 void FMirrorDataTableEditorModule::StartupModule()
 {
 	// Add mirror table specific menus to data table editor toolbar.
@@ -344,6 +363,67 @@ void FMirrorDataTableEditorModule::RegisterMenus()
 				})
 			));
 	}
+
+	// Add mirror data table specific actions to the row right-click menu.
+	if (UToolMenu* RowContextMenu = UToolMenus::Get()->ExtendMenu("DataTableEditor.RowContextMenu"))
+	{
+		RowContextMenu->AddDynamicSection("MirrorDataTableRowDynamicSection", FNewSectionConstructChoice(
+			FNewToolMenuDelegate::CreateLambda([](UToolMenu* InMenu)
+				{
+					if (!InMenu)
+					{
+						return;
+					}
+
+					// Only add mirror-specific row actions when editing a MirrorDataTable.
+					if (!GetMirrorDataTableFromContext(InMenu->Context))
+					{
+						return;
+					}
+
+					const FMirrorTableRow* Row = GetMirrorTableRowFromContext(InMenu->Context);
+					const bool bIsEnabled = Row ? Row->bEnabled : false;
+
+					FToolMenuSection& Section = InMenu->FindOrAddSection("MirrorRowActions", LOCTEXT("MirrorRowActionsLabel", "Mirror"));
+
+					// Toggle this row on/off without removing it from the table.
+					{
+						const FText EnableLabel = LOCTEXT("EnableRowLabel", "Enable Row");
+						const FText DisableLabel = LOCTEXT("DisableRowLabel", "Disable Row");
+						const FText EnableTooltip = LOCTEXT("EnableRowTooltip", "Mark this entry as active so it is used when mirroring animations.");
+						const FText DisableTooltip = LOCTEXT("DisableRowTooltip", "Mark this entry as inactive. It stays in the table but won't be used when mirroring animations.");
+						const FSlateIcon EnableIcon = FSlateIcon(FAppStyle::Get().GetStyleSetName(), "Icons.Visible");
+						const FSlateIcon DisableIcon = FSlateIcon(FAppStyle::Get().GetStyleSetName(), "Icons.Hidden");
+
+						Section.AddMenuEntry(
+							"MDT_ToggleRowEnabled",
+							bIsEnabled ? DisableLabel : EnableLabel,
+							bIsEnabled ? DisableTooltip : EnableTooltip,
+							bIsEnabled ? DisableIcon : EnableIcon,
+							FToolMenuExecuteAction::CreateStatic(&FMirrorDataTableEditorModule::ExecuteMenuAction_ToggleRowEnabled)
+						);
+					}
+
+					// Swap which name is the source and which is the mirror.
+					Section.AddMenuEntry(
+						"MDT_SwapPair",
+						LOCTEXT("SwapPairLabel", "Swap Mirror Pair"),
+						LOCTEXT("SwapPairTooltip", "Swap the source and mirror names for this entry."),
+						FSlateIcon(FAppStyle::Get().GetStyleSetName(), "UMGEditor.Mirror"),
+						FToolMenuExecuteAction::CreateStatic(&FMirrorDataTableEditorModule::ExecuteMenuAction_SwapPair)
+					);
+
+					// Auto-fill the mirrored name using the table's naming rules.
+					Section.AddMenuEntry(
+						"MDT_AutoPairRow",
+						LOCTEXT("AutoPairRowLabel", "Auto Pair Row"),
+						LOCTEXT("AutoPairRowTooltip", "Automatically fill in the mirrored name for this entry using the table's find/replace expressions."),
+						FSlateIcon(FAppStyle::Get().GetStyleSetName(), "AutomationTools.MenuIcon"),
+						FToolMenuExecuteAction::CreateStatic(&FMirrorDataTableEditorModule::ExecuteMenuAction_AutoPairRow)
+					);
+				})
+			));
+	}
 }
 
 void FMirrorDataTableEditorModule::BindSkeletonCallbacks(UMirrorDataTable* MirrorTable)
@@ -583,6 +663,93 @@ void FMirrorDataTableEditorModule::ExecuteMenuAction_ClearSheet(const FToolMenuC
 	}
 }
 
+void FMirrorDataTableEditorModule::ExecuteMenuAction_ToggleRowEnabled(const FToolMenuContext& Context)
+{
+	UMirrorDataTable* MirrorTable = GetMirrorDataTableFromContext(Context);
+	const UDataTableRowMenuContext* RowCtx = Context.FindContext<UDataTableRowMenuContext>();
+	if (!MirrorTable || !RowCtx || !RowCtx->RowDataPtr)
+	{
+		return;
+	}
+
+	FMirrorTableRow* Row = MirrorTable->FindRow<FMirrorTableRow>(RowCtx->RowDataPtr->RowId, TEXT(""));
+	if (!Row)
+	{
+		return;
+	}
+
+	FScopedTransaction Transaction(Row->bEnabled ? LOCTEXT("UndoAction_DisableRow", "Disable Row") : LOCTEXT("UndoAction_EnableRow", "Enable Row"));
+	MirrorTable->Modify();
+
+	// Toggle enabled state and rebuild runtime mirror maps.
+	Row->bEnabled = !Row->bEnabled;
+	MirrorTable->OnDataTableChanged().Broadcast();
+
+	// Let editor know we've changed the data and refresh view.
+	FDataTableEditorUtils::BroadcastPostChange(MirrorTable, FDataTableEditorUtils::EDataTableChangeInfo::RowData);
+	MirrorTable->MarkPackageDirty();
+}
+
+void FMirrorDataTableEditorModule::ExecuteMenuAction_SwapPair(const FToolMenuContext& Context)
+{
+	UMirrorDataTable* MirrorTable = GetMirrorDataTableFromContext(Context);
+	const UDataTableRowMenuContext* RowCtx = Context.FindContext<UDataTableRowMenuContext>();
+	if (!MirrorTable || !RowCtx || !RowCtx->RowDataPtr)
+	{
+		return;
+	}
+
+	FMirrorTableRow* Row = MirrorTable->FindRow<FMirrorTableRow>(RowCtx->RowDataPtr->RowId, TEXT(""));
+	if (!Row)
+	{
+		return;
+	}
+
+	FScopedTransaction Transaction(LOCTEXT("UndoAction_SwapPair", "Swap Mirror Pair"));
+	MirrorTable->Modify();
+
+	// Swap source and mirror names, then rebuild runtime mirror maps.
+	Swap(Row->Name, Row->MirroredName);
+	MirrorTable->OnDataTableChanged().Broadcast();
+
+	// Let editor know we've changed the data and refresh view.
+	FDataTableEditorUtils::BroadcastPostChange(MirrorTable, FDataTableEditorUtils::EDataTableChangeInfo::RowData);
+	MirrorTable->MarkPackageDirty();
+}
+
+void FMirrorDataTableEditorModule::ExecuteMenuAction_AutoPairRow(const FToolMenuContext& Context)
+{
+	UMirrorDataTable* MirrorTable = GetMirrorDataTableFromContext(Context);
+	const UDataTableRowMenuContext* RowCtx = Context.FindContext<UDataTableRowMenuContext>();
+	if (!MirrorTable || !RowCtx || !RowCtx->RowDataPtr)
+	{
+		return;
+	}
+
+	FMirrorTableRow* Row = MirrorTable->FindRow<FMirrorTableRow>(RowCtx->RowDataPtr->RowId, TEXT(""));
+	if (!Row)
+	{
+		return;
+	}
+
+	// Derive the mirrored name using the table's naming rules. No-op if no match is found.
+	const FName MirroredName = MirrorTable->FindReplace(Row->Name);
+	if (MirroredName.IsNone())
+	{
+		return;
+	}
+
+	FScopedTransaction Transaction(LOCTEXT("UndoAction_AutoPairRow", "Auto Pair Row"));
+	MirrorTable->Modify();
+
+	Row->MirroredName = MirroredName;
+	MirrorTable->OnDataTableChanged().Broadcast();
+
+	// Let editor know we've changed the data and refresh view.
+	FDataTableEditorUtils::BroadcastPostChange(MirrorTable, FDataTableEditorUtils::EDataTableChangeInfo::RowData);
+	MirrorTable->MarkPackageDirty();
+}
+
 FSlateIcon FMirrorDataTableEditorModule::GetStatusIcon(UMirrorDataTable* InMirrorDataTable)
 {
 	static const FName ReimportStatusBackground("Persona.AssetClass.Skeleton");
diff --git a/Engine/Source/Editor/MirrorDataTableEditor/Public/MirrorDataTableEditorModule.h b/Engine/Source/Editor/MirrorDataTableEditor/Public/MirrorDataTableEditorModule.h
--- a/Engine/Source/Editor/MirrorDataTableEditor/Public/MirrorDataTableEditorModule.h
+++ b/Engine/Source/Editor/MirrorDataTableEditor/Public/MirrorDataTableEditorModule.h
@@ -37,6 +37,15 @@ public:
 	/** Clears all rows from the active mirror data table. */
 	static void ExecuteMenuAction_ClearSheet(const FToolMenuContext& Context);
 
+	/** Toggles the enabled state of the selected row in the active mirror data table. */
+	static void ExecuteMenuAction_ToggleRowEnabled(const FToolMenuContext& Context);
+
+	/** Swaps the Name and MirroredName of the selected row. */
+	static void ExecuteMenuAction_SwapPair(const FToolMenuContext& Context);
+
+	/** Applies find/replace expressions to the selected row to derive its mirrored name. */
+	static void ExecuteMenuAction_AutoPairRow(const FToolMenuContext& Context);
+
 	/** Returns the status icon for the given mirror data table. */
     static FSlateIcon GetStatusIcon(UMirrorDataTable* InMirrorDataTable);
 
